New modular vehicle example assets plugin

#virtualized

diff --git a/Engine/Plugins/Experimental/ChaosModularVehicleExamples/ChaosModularVehicleExamples.uplugin b/Engine/Plugins/Experimental/ChaosModularVehicleExamples/ChaosModularVehicleExamples.uplugin
new file mode 100644
--- /dev/null
+++ b/Engine/Plugins/Experimental/ChaosModularVehicleExamples/ChaosModularVehicleExamples.uplugin
@@ -0,0 +1,17 @@
+{
+	"FileVersion": 3,
+	"Version": 1,
+	"VersionName": "1.0",
+	"FriendlyName": "Chaos Modular Vehicle Examples",
+	"Description": "Modular Vehicle Example Assets",
+	"Category": "Physics",
+	"CreatedBy": "Epic Games, Inc.",
+	"CreatedByURL": "https://epicgames.com",
+	"DocsURL": "",
+	"MarketplaceURL": "",
+	"SupportURL": "",
+	"CanContainContent": true,
+	"IsBetaVersion": true,
+	"IsExperimentalVersion": false,
+	"Installed": false
+}
\ No newline at end of file
