Update modular vehicle examples to use mover plugin

#Synced-CL 48719260

diff --git a/Engine/Plugins/Experimental/ChaosModularVehicleExamples/ChaosModularVehicleExamples.uplugin b/Engine/Plugins/Experimental/ChaosModularVehicleExamples/ChaosModularVehicleExamples.uplugin
--- a/Engine/Plugins/Experimental/ChaosModularVehicleExamples/ChaosModularVehicleExamples.uplugin
+++ b/Engine/Plugins/Experimental/ChaosModularVehicleExamples/ChaosModularVehicleExamples.uplugin
@@ -10,8 +10,19 @@
 	"DocsURL": "",
 	"MarketplaceURL": "",
 	"SupportURL": "",
+	"EnabledByDefault": false,
 	"CanContainContent": true,
 	"IsBetaVersion": true,
 	"IsExperimentalVersion": false,
-	"Installed": false
-}
\ No newline at end of file
+	"Installed": false,
+	"Plugins": [
+		{
+			"Name": "Mover",
+			"Enabled": true
+		},
+		{
+			"Name": "ChaosMover",
+			"Enabled": true
+		}
+	]
+}
